Submit Tool plugin for overriding editor source control submit behavior

diff --git a/Engine/Plugins/Experimental/SubmitToolEditor/Source/SubmitToolEditor/Private/SubmitToolEditor.cpp b/Engine/Plugins/Experimental/SubmitToolEditor/Source/SubmitToolEditor/Private/SubmitToolEditor.cpp
new file mode 100644
--- /dev/null
+++ b/Engine/Plugins/Experimental/SubmitToolEditor/Source/SubmitToolEditor/Private/SubmitToolEditor.cpp
@@ -0,0 +1,480 @@
+// Copyright Epic Games, Inc. All Rights Reserved.
+
+#include "SubmitToolEditor.h"
+
+#include "EditorValidatorSubsystem.h"
+#include "ISourceControlWindowsModule.h"
+#include "ISourceControlModule.h"
+#include "ISourceControlProvider.h"
+#include "SourceControlOperations.h"
+#include "SubmitToolEditorSettings.h"
+#include "Misc/Paths.h"
+#include "Containers/Ticker.h"
+
+#define LOCTEXT_NAMESPACE "FSubmitToolEditor"
+
+DEFINE_LOG_CATEGORY(LogSubmitToolEditor);
+
+FAutoConsoleCommand CVarCommandEnableSubmitTool = FAutoConsoleCommand(
+	TEXT("SubmitTool.Enable"),
+	TEXT("Enables the submit tool override"),
+	FConsoleCommandDelegate::CreateLambda([]() {
+		USubmitToolEditorSettings* Settings = GetMutableDefault<USubmitToolEditorSettings>();
+		Settings->bSubmitToolEnabled = true;
+		if (!GIsBuildMachine && !IsRunningCommandlet())
+		{
+			FSubmitToolEditorModule::Get().RegisterSubmitOverrideDelegate(Settings);
+		}
+	})
+);
+
+
+
+void FSubmitToolEditorModule::StartupModule()
+{
+	const USubmitToolEditorSettings* Settings = GetDefault<USubmitToolEditorSettings>();
+	if (Settings->bSubmitToolEnabled && !GIsBuildMachine && !IsRunningCommandlet())
+	{			
+		RegisterSubmitOverrideDelegate(Settings);
+	}
+}
+
+void FSubmitToolEditorModule::ShutdownModule()
+{
+	UnregisterSubmitOverrideDelegate();
+}
+
+void FSubmitToolEditorModule::RegisterSubmitOverrideDelegate(const USubmitToolEditorSettings* InSettings)
+{
+	ISourceControlWindowsModule& SourceControlWindowsModule = ISourceControlWindowsModule::Get();
+	if (!SourceControlWindowsModule.SubmitOverrideDelegate.IsBound())
+	{
+		UE_LOG(LogSubmitToolEditor, Display, TEXT("Registering SubmitTool to handle submissions from the editor"));
+
+		if (InSettings->SubmitToolPath.IsEmpty())
+		{
+			UE_LOG(LogSubmitToolEditor, Error, TEXT("Submit Tool path is empty: '%s'"), *InSettings->SubmitToolPath);
+			return;
+		}
+
+		if (InSettings->SubmitToolArguments.IsEmpty())
+		{
+			UE_LOG(LogSubmitToolEditor, Error, TEXT("Submit Tool Args are empty: '%s'"), *InSettings->SubmitToolArguments);
+			return;
+		}
+
+		UE_LOG(LogSubmitToolEditor, Display, TEXT("Registering Submit Tool with Path:'%s' and Args: '%s'"), *InSettings->SubmitToolPath, *InSettings->SubmitToolArguments);
+		ISourceControlWindowsModule::Get().SubmitOverrideDelegate.BindRaw(this, &FSubmitToolEditorModule::OnSubmitToolOverrideCallback);
+	}
+}
+
+void FSubmitToolEditorModule::UnregisterSubmitOverrideDelegate()
+{
+	if (ISourceControlWindowsModule::Get().SubmitOverrideDelegate.IsBound())
+	{
+		UE_LOG(LogSubmitToolEditor, Display, TEXT("Unregistering Submit Tool."));
+		ISourceControlWindowsModule::Get().SubmitOverrideDelegate.Unbind();
+	}
+}
+
+FSubmitOverrideReply FSubmitToolEditorModule::OnSubmitToolOverrideCallback(SSubmitOverrideParameters InParameters)
+{
+	ISourceControlProvider& SourceControlProvider = ISourceControlModule::Get().GetProvider();
+
+	const USubmitToolEditorSettings* Settings = GetDefault<USubmitToolEditorSettings>();
+	FString NormalizedPath(Settings->SubmitToolPath);
+	FPaths::MakePlatformFilename(NormalizedPath);
+
+	FString Platform(FGenericPlatformMisc::GetUBTPlatform());
+
+	NormalizedPath = NormalizedPath.Replace(TEXT("$(Platform)"), *Platform);
+
+#if PLATFORM_WINDOWS
+	const FString LocalAppData = FPlatformMisc::GetEnvironmentVariable(TEXT("LOCALAPPDATA"));
+#elif PLATFORM_MAC
+	const FString LocalAppData = FPlatformMisc::GetEnvironmentVariable(TEXT("HOME"));
+#elif PLATFORM_LINUX
+	const FString LocalAppData = FPaths::Combine(FPlatformMisc::GetEnvironmentVariable(TEXT("HOME")), TEXT(".local"), TEXT("share"));
+#else
+	static_assert(false);
+#endif
+
+	NormalizedPath = NormalizedPath.Replace(TEXT("$(LocalAppData)"), *LocalAppData);
+
+	if (Platform.Equals(TEXT("Win64")))
+	{
+		if (!NormalizedPath.EndsWith(TEXT(".exe"), ESearchCase::IgnoreCase))
+		{
+			NormalizedPath += TEXT(".exe");
+		}
+	}
+
+	if (Platform.Equals(TEXT("Mac")) || Platform.Equals(TEXT("Linux")))
+	{
+		if (NormalizedPath.EndsWith(TEXT(".exe"), ESearchCase::IgnoreCase))
+		{
+			NormalizedPath += NormalizedPath.Replace(TEXT(".exe"), TEXT(""), ESearchCase::IgnoreCase);
+		}
+	}
+
+	if (!FPaths::FileExists(NormalizedPath) && !FPaths::DirectoryExists(NormalizedPath))
+	{
+		UE_LOG(LogSubmitToolEditor, Error, TEXT("The path is invalid: file does not exist - '%s'"), *NormalizedPath);
+		return FSubmitOverrideReply::Error;
+	}
+
+	const FString& ProviderName = SourceControlProvider.GetName().ToString();
+	if (!ProviderName.Equals(TEXT("Perforce"), ESearchCase::IgnoreCase))
+	{
+		UE_LOG(LogSubmitToolEditor, Warning, TEXT("Current Provider is not supported: '%s'"), *ProviderName);
+		return FSubmitOverrideReply::ProviderNotSupported;
+	}
+
+	if (!FPaths::FileExists(NormalizedPath) && !FPaths::DirectoryExists(NormalizedPath))
+	{
+		UE_LOG(LogSubmitToolEditor, Error, TEXT("The path is invalid: file does not exist - '%s'"), *NormalizedPath);
+		return FSubmitOverrideReply::Error;
+	}
+
+	if (InParameters.ToSubmit.HasSubtype<FString>())
+	{
+		return InvokeSubmitTool(SourceControlProvider, NormalizedPath, Settings->SubmitToolArguments, InParameters.Description, InParameters.ToSubmit.GetSubtype<FString>());
+	}
+
+	if (InParameters.ToSubmit.HasSubtype<TArray<FString>>())
+	{
+		return InvokeSubmitTool(SourceControlProvider, NormalizedPath, Settings->SubmitToolArguments, InParameters.Description, InParameters.ToSubmit.GetSubtype<TArray<FString>>());
+	}
+
+	UE_LOG(LogSubmitToolEditor, Error, TEXT("The parameter is invalid: it does not contains and identifier nor a list of files"));
+	return FSubmitOverrideReply::Error;
+}
+
+FSubmitOverrideReply FSubmitToolEditorModule::InvokeSubmitTool(ISourceControlProvider& InProvider, const FString& InPath, const FString& InArgs, const FString& InDescription, const FString& InIdentifier)
+{
+	if (InIdentifier.IsEmpty())
+	{
+		UE_LOG(LogSubmitToolEditor, Error, TEXT("Identifier was empty."));
+		return FSubmitOverrideReply::Error;
+	}
+
+	FString Description(InDescription);
+
+	const USubmitToolEditorSettings* Settings = GetDefault<USubmitToolEditorSettings>();
+	if (Settings->bEnforceDataValidation)
+	{
+		bool bValidChangelist = GetChangelistValidationResult(InProvider, InIdentifier);
+
+		Description = UpdateValidationTag(Description, bValidChangelist);
+
+		if (!bValidChangelist)
+		{
+			SaveChangelistDescription(InProvider, Description, InIdentifier);
+			return FSubmitOverrideReply::Handled;
+		}
+	}
+
+	FString Identifier;
+	if (EditChangelistDescription(InProvider, Description, InIdentifier, Identifier))
+	{
+		return InvokeSubmitTool(InProvider, InPath, InArgs, Identifier);
+	}
+
+	return FSubmitOverrideReply::Error;
+}
+
+FSubmitOverrideReply FSubmitToolEditorModule::InvokeSubmitTool(ISourceControlProvider& InProvider, const FString& InPath, const FString& InArgs, const FString& InDescription, const TArray<FString>& InFiles)
+{
+	FString Identifier;
+	if (CreateChangelist(InProvider, InDescription, InFiles, Identifier))
+	{
+		return InvokeSubmitTool(InProvider, InPath, InArgs, Identifier);
+	}
+
+	return FSubmitOverrideReply::Error;
+}
+
+FSubmitOverrideReply FSubmitToolEditorModule::InvokeSubmitTool(ISourceControlProvider& InProvider, const FString& InPath, const FString& InArgs, const FString& InIdentifier)
+{
+	if (InIdentifier.IsEmpty())
+	{
+		return FSubmitOverrideReply::Error;
+	}
+
+	FString Port, User, Client, WorkspacePath;
+	if (!GetPerforceParameter(InProvider, Port, User, Client, WorkspacePath))
+	{
+		return FSubmitOverrideReply::Error;
+	}
+
+	FString SubstitutedArgs = InArgs
+		.Replace(TEXT("$(Port)"), *Port)
+		.Replace(TEXT("$(User)"), *User)
+		.Replace(TEXT("$(Client)"), *Client)
+		.Replace(TEXT("$(Changelist)"), *InIdentifier)
+		.Replace(TEXT("$(RootDir)"), *WorkspacePath);
+
+	UE_LOG(LogSubmitToolEditor, Display, TEXT("Invoking submit tool: '%s %s'"), *InPath, *SubstitutedArgs);
+	SubmitToolProcHandle = FPlatformProcess::CreateProc(*InPath, *SubstitutedArgs, true, false, false, nullptr, 0, nullptr, nullptr, nullptr);
+
+	if (!SubmitToolProcHandle.IsValid())
+	{
+		UE_LOG(LogSubmitToolEditor, Error, TEXT("Submit Tool could not be launched."));
+		return FSubmitOverrideReply::Error;
+	}
+
+	FTSTicker::GetCoreTicker().AddTicker(FTickerDelegate::CreateRaw(this, &FSubmitToolEditorModule::Tick));
+
+	return FSubmitOverrideReply::Handled;
+}
+
+bool FSubmitToolEditorModule::GetPerforceParameter(ISourceControlProvider& InProvider, FString& OutPort, FString& OutUser, FString& OutClient, FString& OutWorkspacePath)
+{
+	TMap<ISourceControlProvider::EStatus, FString> Status = InProvider.GetStatus();
+
+	FString* Port = Status.Find(ISourceControlProvider::EStatus::Port);
+	if (Port == nullptr || Port->IsEmpty())
+	{
+		UE_LOG(LogSubmitToolEditor, Error, TEXT("Could not get a Port from ISourceControlProvider's status"));
+		return false;
+	}
+
+	FString* User = Status.Find(ISourceControlProvider::EStatus::User);
+	if (User == nullptr || User->IsEmpty())
+	{
+		UE_LOG(LogSubmitToolEditor, Error, TEXT("Could not get a User from ISourceControlProvider's status"));
+		return false;
+	}
+
+	FString* Client = Status.Find(ISourceControlProvider::EStatus::Client);
+	if (Client == nullptr || Client->IsEmpty())
+	{
+		UE_LOG(LogSubmitToolEditor, Error, TEXT("Could not get a Client from ISourceControlProvider's status"));
+		return false;
+	}
+
+	FString* WorkspaceRoot = Status.Find(ISourceControlProvider::EStatus::WorkspacePath);
+	if (WorkspaceRoot == nullptr || WorkspaceRoot->IsEmpty())
+	{
+		UE_LOG(LogSubmitToolEditor, Error, TEXT("Could not get a the Workspace Path ISourceControlProvider's status"));
+		return false;
+	}
+
+	OutPort = *Port;
+	OutUser = *User;
+	OutClient = *Client;
+	OutWorkspacePath = *WorkspaceRoot;
+
+	return true;
+}
+
+bool FSubmitToolEditorModule::EditChangelistDescription(ISourceControlProvider& InProvider, const FString& InNewChangelistDescription, const FString& InIdentifier, FString& OutIdentifier)
+{
+	// default output value is the current identifier
+	OutIdentifier = InIdentifier;
+
+	TArray<FSourceControlChangelistRef> Changelists = InProvider.GetChangelists(EStateCacheUsage::Use);
+	FSourceControlChangelistStatePtr ChangeListState;
+	for (const FSourceControlChangelistRef& Changelist : Changelists)
+	{
+		if (Changelist->GetIdentifier().Equals(InIdentifier))
+		{
+			ChangeListState = InProvider.GetState(Changelist, EStateCacheUsage::Use);
+			break;
+		}
+	}
+
+	if (!ChangeListState.IsValid())
+	{
+		UE_LOG(LogSubmitToolEditor, Warning, TEXT("Could not find changelist '%s'."), *InIdentifier);
+		return false;
+	}
+
+	if (ChangeListState->SupportsPersistentDescription())
+	{
+		TSharedRef<FEditChangelist> EditChangelistOperation = ISourceControlOperation::Create<FEditChangelist>();
+		EditChangelistOperation->SetDescription(FText::FromString(InNewChangelistDescription));
+
+		ECommandResult::Type EditCLResult = InProvider.Execute(EditChangelistOperation, ChangeListState->GetChangelist(), EConcurrency::Synchronous);
+
+		if (EditCLResult != ECommandResult::Succeeded)
+		{
+			UE_LOG(LogSubmitToolEditor, Warning, TEXT("Could not edit changelist '%s''s description, this is not critical and shall not prevent running the submit tool."), *InIdentifier);
+		}
+
+		return true;
+	}
+	// deal with the default changelist
+	else
+	{
+		TArray<FString> FilesToMove;
+		Algo::Transform(ChangeListState->GetFilesStates(), FilesToMove, [](const TSharedRef<ISourceControlState>& FileState) { return FileState->GetFilename(); });
+
+		TSharedRef<FNewChangelist> NewChangelistOperation = ISourceControlOperation::Create<FNewChangelist>();
+
+		NewChangelistOperation->SetDescription(FText::FromString(InNewChangelistDescription));
+		ECommandResult::Type NewCLResult = InProvider.Execute(NewChangelistOperation, FilesToMove, EConcurrency::Synchronous);
+
+		if (NewCLResult != ECommandResult::Succeeded)
+		{
+			UE_LOG(LogSubmitToolEditor, Warning, TEXT("Could not edit changelist '%s''s description, this is not critical and shall not prevent running the submit tool."), *InIdentifier);
+		}
+
+		const FSourceControlChangelistPtr& NewChangeList = NewChangelistOperation->GetNewChangelist();
+		OutIdentifier = NewChangeList->GetIdentifier();
+		return true;
+	}
+}
+
+void FSubmitToolEditorModule::SaveChangelistDescription(ISourceControlProvider& InProvider, const FString& InNewChangelistDescription, const FString& InIdentifier)
+{
+	TArray<FSourceControlChangelistRef> Changelists = InProvider.GetChangelists(EStateCacheUsage::Use);
+	FSourceControlChangelistStatePtr ChangeListState;
+	for (const FSourceControlChangelistRef& Changelist : Changelists)
+	{
+		if (Changelist->GetIdentifier().Equals(InIdentifier))
+		{
+			ChangeListState = InProvider.GetState(Changelist, EStateCacheUsage::Use);
+			break;
+		}
+	}
+
+	if (!ChangeListState.IsValid())
+	{
+		UE_LOG(LogSubmitToolEditor, Warning, TEXT("Could not find changelist '%s'."), *InIdentifier);
+		return;
+	}
+
+	if (ChangeListState->SupportsPersistentDescription())
+	{
+		TSharedRef<FEditChangelist> EditChangelistOperation = ISourceControlOperation::Create<FEditChangelist>();
+		EditChangelistOperation->SetDescription(FText::FromString(InNewChangelistDescription));
+
+		ECommandResult::Type EditCLResult = InProvider.Execute(EditChangelistOperation, ChangeListState->GetChangelist(), EConcurrency::Synchronous);
+
+		if (EditCLResult != ECommandResult::Succeeded)
+		{
+			UE_LOG(LogSubmitToolEditor, Warning, TEXT("Could not edit changelist '%s''s description, this is not critical and shall not prevent running the submit tool."), *InIdentifier);
+		}
+	}
+	else
+	{
+		// do nothing, the description cannot be changed
+	}
+}
+
+bool FSubmitToolEditorModule::CreateChangelist(ISourceControlProvider & InProvider, const FString & InNewChangelistDescription, const TArray<FString>&InFiles, FString & OutIdentifier)
+{
+	TSharedRef<FNewChangelist> NewChangelistOperation = ISourceControlOperation::Create<FNewChangelist>();
+	NewChangelistOperation->SetDescription(FText::FromString(InNewChangelistDescription));
+
+	ECommandResult::Type CreateCLResult = InProvider.Execute(NewChangelistOperation, InFiles, EConcurrency::Synchronous);
+
+	switch (CreateCLResult)
+	{
+		case ECommandResult::Type::Succeeded:
+		{
+			FSourceControlChangelistPtr NewChangeList = NewChangelistOperation->GetNewChangelist();
+			OutIdentifier = NewChangeList->GetIdentifier();
+			return true;
+		}
+		case ECommandResult::Type::Failed:
+		case ECommandResult::Type::Cancelled:
+		default:
+			UE_LOG(LogSubmitToolEditor, Error, TEXT("Could not create changelist."));
+			break;
+	}
+
+	return false;
+}
+
+bool FSubmitToolEditorModule::GetChangelistValidationResult(ISourceControlProvider& InProvider, const FString& InIdentifier)
+{
+	if (InIdentifier.IsEmpty())
+	{
+		return true;
+	}
+
+	TArray<FSourceControlChangelistRef> Changelists = InProvider.GetChangelists(EStateCacheUsage::Use);
+	TSharedPtr<FSourceControlChangelistPtr> ChangelistToProcess = nullptr;
+	for (const FSourceControlChangelistRef& Changelist : Changelists)
+	{
+		if (Changelist->GetIdentifier().Equals(InIdentifier))
+		{
+			ChangelistToProcess = MakeShared<FSourceControlChangelistPtr>(Changelist);
+			break;
+		}
+	}
+	if (ChangelistToProcess->IsValid())
+	{
+		FSourceControlPreSubmitDataValidationDelegate ValidationDelegate = ISourceControlModule::Get().GetRegisteredPreSubmitDataValidation();
+
+		EDataValidationResult ValidationResult = EDataValidationResult::NotValidated;
+		TArray<FText> ValidationErrors;
+		TArray<FText> ValidationWarnings;
+
+		ValidationDelegate.ExecuteIfBound(*(ChangelistToProcess), ValidationResult, ValidationErrors, ValidationWarnings);
+
+		return (ValidationResult == EDataValidationResult::Valid);
+	}
+	else
+	{
+		UE_LOG(LogSubmitToolEditor, Warning, TEXT("Could not find changelist '%s'."), *InIdentifier);
+	}
+
+	return false;
+}
+
+FString FSubmitToolEditorModule::UpdateValidationTag(const FString& InDescription, bool bIsvalid)
+{
+	if (USourceControlPreferences::IsValidationTagEnabled())
+	{
+		const FText ValidationTag = LOCTEXT("ValidationTag", "#changelist validated");
+		const FString ValidationTagStr = ValidationTag.ToString();
+
+		if (!bIsvalid)
+		{
+			return InDescription.Replace(*ValidationTagStr, TEXT(""), ESearchCase::IgnoreCase);
+		}
+		else
+		{
+			if (InDescription.Find(ValidationTagStr, ESearchCase::IgnoreCase) == INDEX_NONE)
+			{
+				FStringOutputDevice Str;
+
+				Str.SetAutoEmitLineTerminator(true);
+				Str.Log(InDescription);
+				Str.Log(ValidationTagStr);
+
+				return Str;
+			}
+		}
+	}
+
+	return InDescription;
+}
+
+bool FSubmitToolEditorModule::Tick(float InDeltaTime)
+{
+	if (!FPlatformProcess::IsProcRunning(SubmitToolProcHandle))
+	{
+		ISourceControlProvider& InProvider = ISourceControlModule::Get().GetProvider();
+
+		// When submit tool is done, trigger a cache refresh for CLs and File status so that the editor can refresh its UI
+		TArray<FSourceControlChangelistRef> Cls = InProvider.GetChangelists(EStateCacheUsage::ForceUpdate);
+
+		TSharedRef<FUpdatePendingChangelistsStatus, ESPMode::ThreadSafe> UpdatePendingChangelistsOperation = ISourceControlOperation::Create<FUpdatePendingChangelistsStatus>();
+		UpdatePendingChangelistsOperation->SetUpdateAllChangelists(true);
+		UpdatePendingChangelistsOperation->SetUpdateFilesStates(true);
+		UpdatePendingChangelistsOperation->SetUpdateShelvedFilesStates(true);
+
+		InProvider.Execute(UpdatePendingChangelistsOperation, EConcurrency::Asynchronous);
+		return false;
+	}
+
+	return true;
+}
+
+#undef LOCTEXT_NAMESPACE
+
+IMPLEMENT_MODULE(FSubmitToolEditorModule, SubmitToolEditor)
\ No newline at end of file
diff --git a/Engine/Plugins/Experimental/SubmitToolEditor/Source/SubmitToolEditor/Private/SubmitToolEditor.h b/Engine/Plugins/Experimental/SubmitToolEditor/Source/SubmitToolEditor/Private/SubmitToolEditor.h
new file mode 100644
--- /dev/null
+++ b/Engine/Plugins/Experimental/SubmitToolEditor/Source/SubmitToolEditor/Private/SubmitToolEditor.h
@@ -0,0 +1,52 @@
+// Copyright Epic Games, Inc. All Rights Reserved.
+
+#pragma once
+
+#include "CoreMinimal.h"
+#include "Modules/ModuleInterface.h"
+#include "Modules/ModuleManager.h"
+#include "HAL/PlatformProcess.h"
+
+DECLARE_LOG_CATEGORY_EXTERN(LogSubmitToolEditor, Log, All);
+
+enum class FSubmitOverrideReply;
+class USubmitToolEditorSettings;
+class ISourceControlProvider;
+struct SSubmitOverrideParameters;
+
+class FSubmitToolEditorModule : public IModuleInterface
+{
+public:
+	/**
+	 * Get reference to the SubmitToolEditor module instance
+	 */
+	static inline FSubmitToolEditorModule& Get()
+	{
+		return FModuleManager::LoadModuleChecked<FSubmitToolEditorModule>("SubmitToolEditor");
+	}
+
+	/** IModuleInterface implementation */
+	virtual void StartupModule() override;
+	virtual void ShutdownModule() override;
+
+	void RegisterSubmitOverrideDelegate(const USubmitToolEditorSettings* InSettings);
+	void UnregisterSubmitOverrideDelegate();
+	
+private:
+
+	FSubmitOverrideReply OnSubmitToolOverrideCallback(SSubmitOverrideParameters InParameters);
+	FSubmitOverrideReply InvokeSubmitTool(ISourceControlProvider& InProvider, const FString& InPath, const FString& InArgs, const FString& InDescription, const FString& InIdentifier);
+	FSubmitOverrideReply InvokeSubmitTool(ISourceControlProvider& InProvider, const FString& InPath, const FString& InArgs, const FString& InDescription, const TArray<FString>& InFiles);
+	FSubmitOverrideReply InvokeSubmitTool(ISourceControlProvider& InProvider, const FString& InPath, const FString& InArgs, const FString& InIdentifier);
+
+	bool GetPerforceParameter(ISourceControlProvider& InProvider, FString& OutPort, FString& OutUser, FString& OutClient, FString& OutWorkspacePath);
+	bool EditChangelistDescription(ISourceControlProvider& InProvider, const FString& InNewChangelistDescription, const FString& InIdentifier, FString& OutIdentifier);
+	void SaveChangelistDescription(ISourceControlProvider& InProvider, const FString& InNewChangelistDescription, const FString& InIdentifier);
+	bool CreateChangelist(ISourceControlProvider& InProvider, const FString& InNewChangelistDescription, const TArray<FString>& InFiles, FString& OutIdentifier);
+	bool GetChangelistValidationResult(ISourceControlProvider& InProvider, const FString& InIdentifier);
+	FString UpdateValidationTag(const FString& InDescription, bool bIsvalid);
+
+	bool Tick(float InDeltaTime);
+
+	FProcHandle SubmitToolProcHandle;
+};
diff --git a/Engine/Plugins/Experimental/SubmitToolEditor/Source/SubmitToolEditor/Private/SubmitToolEditorSettings.cpp b/Engine/Plugins/Experimental/SubmitToolEditor/Source/SubmitToolEditor/Private/SubmitToolEditorSettings.cpp
new file mode 100644
--- /dev/null
+++ b/Engine/Plugins/Experimental/SubmitToolEditor/Source/SubmitToolEditor/Private/SubmitToolEditorSettings.cpp
@@ -0,0 +1,30 @@
+// Copyright Epic Games, Inc. All Rights Reserved.
+
+#include "SubmitToolEditorSettings.h"
+
+#include "SubmitToolEditor.h"
+
+USubmitToolEditorSettings::USubmitToolEditorSettings()
+{
+}
+
+void USubmitToolEditorSettings::PostEditChangeProperty(FPropertyChangedEvent& PropertyChangedEvent)
+{
+	Super::PostEditChangeProperty(PropertyChangedEvent);
+
+	if (bSubmitToolEnabled && GIsBuildMachine == false && IsRunningCommandlet() == false)
+	{
+		FSubmitToolEditorModule::Get().RegisterSubmitOverrideDelegate(this);
+	}
+	else
+	{
+		FSubmitToolEditorModule::Get().UnregisterSubmitOverrideDelegate();
+	}
+
+	SaveConfig();
+}
+
+FName USubmitToolEditorSettings::GetCategoryName() const
+{
+	return FName("Editor");
+}
diff --git a/Engine/Plugins/Experimental/SubmitToolEditor/Source/SubmitToolEditor/Private/SubmitToolEditorSettings.h b/Engine/Plugins/Experimental/SubmitToolEditor/Source/SubmitToolEditor/Private/SubmitToolEditorSettings.h
new file mode 100644
--- /dev/null
+++ b/Engine/Plugins/Experimental/SubmitToolEditor/Source/SubmitToolEditor/Private/SubmitToolEditorSettings.h
@@ -0,0 +1,35 @@
+// Copyright Epic Games, Inc. All Rights Reserved.
+
+#pragma once
+
+#include "Engine/DeveloperSettings.h"
+
+#include "SubmitToolEditorSettings.generated.h"
+
+struct FPropertyChangedEvent;
+
+/** Settings for the submit tool in the editor */
+
+UCLASS(config=Editor, meta = (DisplayName = "Submit Tool Settings"))
+class USubmitToolEditorSettings : public UDeveloperSettings
+{
+	GENERATED_BODY()
+
+public:
+	USubmitToolEditorSettings();
+	void PostEditChangeProperty(FPropertyChangedEvent& PropertyChangedEvent) override;
+	virtual FName GetCategoryName() const override;
+
+	UPROPERTY(config, Category = "Submit Tool", EditDefaultsOnly)
+	FString SubmitToolPath;
+
+	UPROPERTY(config, Category = "Submit Tool", EditDefaultsOnly)
+	FString SubmitToolArguments;
+
+	UPROPERTY(config, Category = "Submit Tool", EditDefaultsOnly)
+	bool bSubmitToolEnabled;
+
+	UPROPERTY(config, Category = "Submit Tool", EditDefaultsOnly)
+	bool bEnforceDataValidation;
+};
+
diff --git a/Engine/Plugins/Experimental/SubmitToolEditor/Source/SubmitToolEditor/SubmitToolEditor.Build.cs b/Engine/Plugins/Experimental/SubmitToolEditor/Source/SubmitToolEditor/SubmitToolEditor.Build.cs
new file mode 100644
--- /dev/null
+++ b/Engine/Plugins/Experimental/SubmitToolEditor/Source/SubmitToolEditor/SubmitToolEditor.Build.cs
@@ -0,0 +1,23 @@
+// Copyright Epic Games, Inc. All Rights Reserved.
+using UnrealBuildTool;
+
+public class SubmitToolEditor : ModuleRules
+{
+	public SubmitToolEditor(ReadOnlyTargetRules Target) : base(Target)
+	{
+		PCHUsage = ModuleRules.PCHUsageMode.UseExplicitOrSharedPCHs;
+
+		PrivateDependencyModuleNames.AddRange(
+			new string[]
+			{
+				"Core",
+				"CoreUObject",
+				"Engine",
+				"SourceControlWindows",
+				"SourceControl",
+				"DeveloperSettings",
+				"DataValidation"
+			}
+		);
+	}
+}
\ No newline at end of file
diff --git a/Engine/Plugins/Experimental/SubmitToolEditor/SubmitToolEditor.uplugin b/Engine/Plugins/Experimental/SubmitToolEditor/SubmitToolEditor.uplugin
new file mode 100644
--- /dev/null
+++ b/Engine/Plugins/Experimental/SubmitToolEditor/SubmitToolEditor.uplugin
@@ -0,0 +1,29 @@
+﻿{
+	"FileVersion": 3,
+	"Version": 1,
+	"VersionName": "1.0",
+	"FriendlyName": "Submit Tool Editor Override",
+	"Description": "Sets up Submit Tool to be launched by the editor",
+	"Category": "Editor",
+	"CreatedBy": "Epic Games, Inc.",
+	"CreatedByURL": "https://epicgames.com",
+	"CanContainContent": false,
+	"IsBetaVersion": false,
+	"IsExperimentalVersion": true,
+	"Installed": false,
+	"EnabledByDefault": false,
+	"Modules": [
+		{
+			"Name": "SubmitToolEditor",
+			"Type": "Editor",
+			"LoadingPhase": "Default"
+		}
+	],
+	 "Plugins": [
+		{
+		"Name": "DataValidation",
+		"Enabled": true,
+		"Optional": false
+		}
+	 ]
+}
\ No newline at end of file
