Add check for Source Control Module being valid before trying to unregister

diff --git a/Engine/Plugins/Experimental/SubmitToolEditor/Source/SubmitToolEditor/Private/SubmitToolEditor.cpp b/Engine/Plugins/Experimental/SubmitToolEditor/Source/SubmitToolEditor/Private/SubmitToolEditor.cpp
--- a/Engine/Plugins/Experimental/SubmitToolEditor/Source/SubmitToolEditor/Private/SubmitToolEditor.cpp
+++ b/Engine/Plugins/Experimental/SubmitToolEditor/Source/SubmitToolEditor/Private/SubmitToolEditor.cpp
@@ -70,10 +70,12 @@ void FSubmitToolEditorModule::RegisterSubmitOverrideDelegate(const USubmitToolEd
 
 void FSubmitToolEditorModule::UnregisterSubmitOverrideDelegate()
 {
-	if (ISourceControlWindowsModule::Get().SubmitOverrideDelegate.IsBound())
+	ISourceControlWindowsModule* SourceControlModule = ISourceControlWindowsModule::TryGet();
+
+	if (SourceControlModule != nullptr && SourceControlModule->SubmitOverrideDelegate.IsBound())
 	{
 		UE_LOG(LogSubmitToolEditor, Display, TEXT("Unregistering Submit Tool."));
-		ISourceControlWindowsModule::Get().SubmitOverrideDelegate.Unbind();
+		SourceControlModule->SubmitOverrideDelegate.Unbind();
 	}
 }
 
