Adding disable command for submit tool plugin

diff --git a/Engine/Plugins/Experimental/SubmitToolEditor/Source/SubmitToolEditor/Private/SubmitToolEditor.cpp b/Engine/Plugins/Experimental/SubmitToolEditor/Source/SubmitToolEditor/Private/SubmitToolEditor.cpp
--- a/Engine/Plugins/Experimental/SubmitToolEditor/Source/SubmitToolEditor/Private/SubmitToolEditor.cpp
+++ b/Engine/Plugins/Experimental/SubmitToolEditor/Source/SubmitToolEditor/Private/SubmitToolEditor.cpp
@@ -29,6 +29,18 @@ FAutoConsoleCommand CVarCommandEnableSubmitTool = FAutoConsoleCommand(
 );
 
 
+FAutoConsoleCommand CVarCommandDisableSubmitTool = FAutoConsoleCommand(
+	TEXT("SubmitTool.Disable"),
+	TEXT("Disable the submit tool override"),
+	FConsoleCommandDelegate::CreateLambda([]() {
+		USubmitToolEditorSettings* Settings = GetMutableDefault<USubmitToolEditorSettings>();
+		Settings->bSubmitToolEnabled = false;
+		if (!GIsBuildMachine && !IsRunningCommandlet())
+		{
+			FSubmitToolEditorModule::Get().UnregisterSubmitOverrideDelegate();
+		}
+		})
+);
 
 void FSubmitToolEditorModule::StartupModule()
 {
