Check the source control system compatibility before trying to find submit tool exe

diff --git a/Engine/Plugins/Experimental/SubmitToolEditor/Source/SubmitToolEditor/Private/SubmitToolEditor.cpp b/Engine/Plugins/Experimental/SubmitToolEditor/Source/SubmitToolEditor/Private/SubmitToolEditor.cpp
--- a/Engine/Plugins/Experimental/SubmitToolEditor/Source/SubmitToolEditor/Private/SubmitToolEditor.cpp
+++ b/Engine/Plugins/Experimental/SubmitToolEditor/Source/SubmitToolEditor/Private/SubmitToolEditor.cpp
@@ -95,6 +95,13 @@ FSubmitOverrideReply FSubmitToolEditorModule::OnSubmitToolOverrideCallback(SSubm
 {
 	ISourceControlProvider& SourceControlProvider = ISourceControlModule::Get().GetProvider();
 
+	const FString& ProviderName = SourceControlProvider.GetName().ToString();
+	if (!ProviderName.Equals(TEXT("Perforce"), ESearchCase::IgnoreCase))
+	{
+		UE_LOG(LogSubmitToolEditor, Warning, TEXT("Current Provider is not supported for Submit Tool: '%s', Using the regular editor source control flow"), *ProviderName);
+		return FSubmitOverrideReply::ProviderNotSupported;
+	}
+
 	const USubmitToolEditorSettings* Settings = GetDefault<USubmitToolEditorSettings>();
 	FString NormalizedPath(Settings->SubmitToolPath);
 	FPaths::MakePlatformFilename(NormalizedPath);
@@ -137,13 +144,6 @@ FSubmitOverrideReply FSubmitToolEditorModule::OnSubmitToolOverrideCallback(SSubm
 		return FSubmitOverrideReply::Error;
 	}
 
-	const FString& ProviderName = SourceControlProvider.GetName().ToString();
-	if (!ProviderName.Equals(TEXT("Perforce"), ESearchCase::IgnoreCase))
-	{
-		UE_LOG(LogSubmitToolEditor, Warning, TEXT("Current Provider is not supported: '%s'"), *ProviderName);
-		return FSubmitOverrideReply::ProviderNotSupported;
-	}
-
 	if (!FPaths::FileExists(NormalizedPath) && !FPaths::DirectoryExists(NormalizedPath))
 	{
 		UE_LOG(LogSubmitToolEditor, Error, TEXT("The path is invalid: file does not exist - '%s'"), *NormalizedPath);
