Added override arguments for path and args to SubmitTool.Enable command

diff --git a/Engine/Plugins/Experimental/SubmitToolEditor/Source/SubmitToolEditor/Private/SubmitToolEditor.cpp b/Engine/Plugins/Experimental/SubmitToolEditor/Source/SubmitToolEditor/Private/SubmitToolEditor.cpp
--- a/Engine/Plugins/Experimental/SubmitToolEditor/Source/SubmitToolEditor/Private/SubmitToolEditor.cpp
+++ b/Engine/Plugins/Experimental/SubmitToolEditor/Source/SubmitToolEditor/Private/SubmitToolEditor.cpp
@@ -18,8 +18,23 @@ DEFINE_LOG_CATEGORY(LogSubmitToolEditor);
 FAutoConsoleCommand CVarCommandEnableSubmitTool = FAutoConsoleCommand(
 	TEXT("SubmitTool.Enable"),
 	TEXT("Enables the submit tool override"),
-	FConsoleCommandDelegate::CreateLambda([]() {
+	FConsoleCommandWithArgsDelegate::CreateLambda([](const TArray<FString>& Args) {
 		USubmitToolEditorSettings* Settings = GetMutableDefault<USubmitToolEditorSettings>();
+		
+		if (!Args.IsEmpty())
+		{
+			if (!Args[0].IsEmpty())
+			{
+				Settings->SubmitToolPath = Args[0];
+			}
+
+			if (Args.Num() >= 2 && !Args[1].IsEmpty())
+			{
+				Settings->SubmitToolArguments = Args[1];
+			}
+
+		}
+
 		Settings->bSubmitToolEnabled = true;
 		if (!GIsBuildMachine && !IsRunningCommandlet())
 		{
