Added option for forcing submit tool enabled by default, setting this to false allows for a soft activation where users that don't have submit tool installed are allowed to submit

#rnx

diff --git a/Engine/Plugins/Experimental/SubmitToolEditor/Source/SubmitToolEditor/Private/SubmitToolEditor.cpp b/Engine/Plugins/Experimental/SubmitToolEditor/Source/SubmitToolEditor/Private/SubmitToolEditor.cpp
--- a/Engine/Plugins/Experimental/SubmitToolEditor/Source/SubmitToolEditor/Private/SubmitToolEditor.cpp
+++ b/Engine/Plugins/Experimental/SubmitToolEditor/Source/SubmitToolEditor/Private/SubmitToolEditor.cpp
@@ -33,6 +33,10 @@ FAutoConsoleCommand CVarCommandEnableSubmitTool = FAutoConsoleCommand(
 				Settings->SubmitToolArguments = Args[1];
 			}
 
+			if (Args.Num() >= 3)
+			{
+				Settings->bForceSubmitTool = Args[2].ToBool();
+			}
 		}
 
 		Settings->bSubmitToolEnabled = true;
@@ -156,13 +160,14 @@ FSubmitOverrideReply FSubmitToolEditorModule::OnSubmitToolOverrideCallback(SSubm
 	if (!FPaths::FileExists(NormalizedPath) && !FPaths::DirectoryExists(NormalizedPath))
 	{
 		UE_LOG(LogSubmitToolEditor, Error, TEXT("The path is invalid: file does not exist - '%s'"), *NormalizedPath);
-		return FSubmitOverrideReply::Error;
+		return Settings->bForceSubmitTool ? FSubmitOverrideReply::Error : FSubmitOverrideReply::ProviderNotSupported;
 	}
 
 	if (!FPaths::FileExists(NormalizedPath) && !FPaths::DirectoryExists(NormalizedPath))
 	{
-		UE_LOG(LogSubmitToolEditor, Error, TEXT("The path is invalid: file does not exist - '%s'"), *NormalizedPath);
-		return FSubmitOverrideReply::Error;
+		UE_LOG(LogSubmitToolEditor, Error, TEXT("The path is invalid: file does not exist - '%s'"), *NormalizedPath);		
+		return Settings->bForceSubmitTool ? FSubmitOverrideReply::Error : FSubmitOverrideReply::ProviderNotSupported;
+
 	}
 
 	if (InParameters.ToSubmit.HasSubtype<FString>())
diff --git a/Engine/Plugins/Experimental/SubmitToolEditor/Source/SubmitToolEditor/Private/SubmitToolEditorSettings.h b/Engine/Plugins/Experimental/SubmitToolEditor/Source/SubmitToolEditor/Private/SubmitToolEditorSettings.h
--- a/Engine/Plugins/Experimental/SubmitToolEditor/Source/SubmitToolEditor/Private/SubmitToolEditorSettings.h
+++ b/Engine/Plugins/Experimental/SubmitToolEditor/Source/SubmitToolEditor/Private/SubmitToolEditorSettings.h
@@ -28,7 +28,10 @@ public:
 
 	UPROPERTY(config, Category = "Submit Tool", EditDefaultsOnly)
 	bool bSubmitToolEnabled;
-
+	
+	UPROPERTY(config, Category = "Submit Tool", EditDefaultsOnly)
+	bool bForceSubmitTool = true;
+	
 	UPROPERTY(config, Category = "Submit Tool", EditDefaultsOnly)
 	bool bEnforceDataValidation;
 };
