Add validation message parsing to SubmitToolEditor.
This can optionally generate additional command line arguments for the SubmitTool.
Also it can optionally update changelist tags.
#rnx

diff --git a/Engine/Plugins/Experimental/SubmitToolEditor/Source/SubmitToolEditor/Private/SubmitToolEditor.cpp b/Engine/Plugins/Experimental/SubmitToolEditor/Source/SubmitToolEditor/Private/SubmitToolEditor.cpp
--- a/Engine/Plugins/Experimental/SubmitToolEditor/Source/SubmitToolEditor/Private/SubmitToolEditor.cpp
+++ b/Engine/Plugins/Experimental/SubmitToolEditor/Source/SubmitToolEditor/Private/SubmitToolEditor.cpp
@@ -2,15 +2,16 @@
 
 #include "SubmitToolEditor.h"
 
+#include "Containers/Ticker.h"
 #include "EditorValidatorSubsystem.h"
+#include "Internationalization/Regex.h"
 #include "ISourceControlWindowsModule.h"
 #include "ISourceControlModule.h"
 #include "ISourceControlProvider.h"
-#include "SourceControlOperations.h"
-#include "SubmitToolEditorSettings.h"
 #include "Misc/Paths.h"
 #include "Misc/StringOutputDevice.h"
-#include "Containers/Ticker.h"
+#include "SourceControlOperations.h"
+#include "SubmitToolEditorSettings.h"
 
 #define LOCTEXT_NAMESPACE "FSubmitToolEditor"
 
@@ -196,12 +197,15 @@ FSubmitOverrideReply FSubmitToolEditorModule::InvokeSubmitTool(ISourceControlPro
 		return FSubmitOverrideReply::Error;
 	}
 
+	FString Args(InArgs);
 	FString Description(InDescription);
-
+	
 	const USubmitToolEditorSettings* Settings = GetDefault<USubmitToolEditorSettings>();
 	if (Settings->bEnforceDataValidation)
 	{
-		bool bValidChangelist = GetChangelistValidationResult(InProvider, InIdentifier);
+		TArray<FText> ValidationMessages;
+
+		bool bValidChangelist = GetChangelistValidationResult(InProvider, InIdentifier, ValidationMessages);
 
 		Description = UpdateValidationTag(Description, bValidChangelist);
 
@@ -210,12 +214,15 @@ FSubmitOverrideReply FSubmitToolEditorModule::InvokeSubmitTool(ISourceControlPro
 			SaveChangelistDescription(InProvider, Description, InIdentifier);
 			return FSubmitOverrideReply::Handled;
 		}
+
+		UpdateOptionalValidationArgs(Settings, ValidationMessages, Args);
+		UpdateOptionalValidationTags(Settings, ValidationMessages, Description);
 	}
 
 	FString Identifier;
 	if (EditChangelistDescription(InProvider, Description, InIdentifier, Identifier))
 	{
-		return InvokeSubmitTool(InProvider, InPath, InArgs, Identifier);
+		return InvokeSubmitTool(InProvider, InPath, Args, Identifier);
 	}
 
 	return FSubmitOverrideReply::Error;
@@ -439,7 +446,7 @@ bool FSubmitToolEditorModule::CreateChangelist(ISourceControlProvider & InProvid
 	return false;
 }
 
-bool FSubmitToolEditorModule::GetChangelistValidationResult(ISourceControlProvider& InProvider, const FString& InIdentifier)
+bool FSubmitToolEditorModule::GetChangelistValidationResult(ISourceControlProvider& InProvider, const FString& InIdentifier, TArray<FText>& OutValidationMessages)
 {
 	if (InIdentifier.IsEmpty())
 	{
@@ -464,6 +471,7 @@ bool FSubmitToolEditorModule::GetChangelistValidationResult(ISourceControlProvid
 		
 		ValidationDelegate.ExecuteIfBound(*(ChangelistToProcess), ValidationResult);
 
+		OutValidationMessages = MoveTemp(ValidationResult.AllMessages);
 		return (ValidationResult.Result == EDataValidationResult::Valid);
 	}
 	else
@@ -503,6 +511,89 @@ FString FSubmitToolEditorModule::UpdateValidationTag(const FString& InDescriptio
 	return InDescription;
 }
 
+void FSubmitToolEditorModule::UpdateOptionalValidationArgs(USubmitToolEditorSettings const* InSettings, TArrayView<FText> InValidationMessages, FString& InOutArgs)
+{
+	for (FSubmitToolArgumentTrigger const& Trigger : InSettings->OptionalArgumentTriggers)
+	{
+		TSharedPtr<FRegexPattern> Pattern = MakeShared<FRegexPattern>(Trigger.RegExMessage, ERegexPatternFlags::CaseInsensitive);
+
+		for (FText const& Message : InValidationMessages)
+		{
+			FRegexMatcher Regex = FRegexMatcher(*Pattern, Message.ToString());
+			if (Regex.FindNext())
+			{
+				InOutArgs.Appendf(TEXT(" %s"), *Trigger.SubmitToolArgument);
+				break;
+			}
+		}
+	}
+}
+
+void FSubmitToolEditorModule::UpdateOptionalValidationTags(USubmitToolEditorSettings const* InSettings, TArrayView<FText> InValidationMessages, FString& InOutDescription)
+{
+	if (InSettings->ChangelistTagTriggers.Num() == 0)
+	{
+		return;
+	}
+
+	TArray<FString> Lines;
+	InOutDescription.ParseIntoArrayLines(Lines, /*bCullEmpty*/false);
+
+	for (FSubmitToolChangelistTagTrigger const& Trigger : InSettings->ChangelistTagTriggers)
+	{
+		// Find matching message.
+		TSharedPtr<FRegexPattern> Pattern = MakeShared<FRegexPattern>(Trigger.RegExMessage, ERegexPatternFlags::CaseInsensitive);
+
+		FString TagLine;
+		for (int32 MessageIndex = InValidationMessages.Num() - 1; MessageIndex >= 0; --MessageIndex)
+		{
+			const FString Message = InValidationMessages[MessageIndex].ToString();
+			FRegexMatcher Regex = FRegexMatcher(*Pattern, Message);
+			if (Regex.FindNext())
+			{
+				const FString TagValue = Regex.GetCaptureGroup(1);
+				if (!TagValue.IsEmpty())
+				{
+					TagLine = FString::Printf(TEXT("#%s %s"), *Trigger.SubmitToolTag, *TagValue);
+					break;
+				}
+			}
+		}
+
+		// Find lines with matching tag.
+		const FString TagLineStart = TEXT("#") + Trigger.SubmitToolTag;
+
+		for (int32 LineIndex = Lines.Num() - 1; LineIndex >= 0 ; --LineIndex)
+		{
+			if (!Lines[LineIndex].StartsWith(TagLineStart, ESearchCase::IgnoreCase))
+			{
+				continue;
+			}
+
+			if (TagLine.IsEmpty())
+			{
+				// Remove existing tag.
+				Lines.RemoveAt(LineIndex);
+			}
+			else
+			{
+				// Replace existing tag.
+				Lines[LineIndex] = MoveTemp(TagLine);
+				// Reset and continue iterating to remove any additional matching tags.
+				TagLine.Reset();
+			}
+		}
+
+		if (!TagLine.IsEmpty())
+		{
+			// Didn't replace existing tag. So add new tag line.
+			Lines.Add(TagLine);
+		}
+	}
+
+	InOutDescription = FString::Join(Lines, LINE_TERMINATOR);
+}
+
 bool FSubmitToolEditorModule::Tick(float InDeltaTime)
 {
 	if (!FPlatformProcess::IsProcRunning(SubmitToolProcHandle))
diff --git a/Engine/Plugins/Experimental/SubmitToolEditor/Source/SubmitToolEditor/Private/SubmitToolEditor.h b/Engine/Plugins/Experimental/SubmitToolEditor/Source/SubmitToolEditor/Private/SubmitToolEditor.h
--- a/Engine/Plugins/Experimental/SubmitToolEditor/Source/SubmitToolEditor/Private/SubmitToolEditor.h
+++ b/Engine/Plugins/Experimental/SubmitToolEditor/Source/SubmitToolEditor/Private/SubmitToolEditor.h
@@ -45,8 +45,10 @@ private:
 	bool EditChangelistDescription(ISourceControlProvider& InProvider, const FString& InNewChangelistDescription, const FString& InIdentifier, FString& OutIdentifier);
 	void SaveChangelistDescription(ISourceControlProvider& InProvider, const FString& InNewChangelistDescription, const FString& InIdentifier);
 	bool CreateChangelist(ISourceControlProvider& InProvider, const FString& InNewChangelistDescription, const TArray<FString>& InFiles, FString& OutIdentifier);
-	bool GetChangelistValidationResult(ISourceControlProvider& InProvider, const FString& InIdentifier);
+	bool GetChangelistValidationResult(ISourceControlProvider& InProvider, const FString& InIdentifier, TArray<FText>& OutValidationMessages);
 	FString UpdateValidationTag(const FString& InDescription, bool bIsvalid);
+	void UpdateOptionalValidationArgs(USubmitToolEditorSettings const* InSettings, TArrayView<FText> InValidationMessages, FString& InOutArgs);
+	void UpdateOptionalValidationTags(USubmitToolEditorSettings const* InSettings, TArrayView<FText> InValidationMessages, FString& InOutDescription);
 
 	bool Tick(float InDeltaTime);
 
diff --git a/Engine/Plugins/Experimental/SubmitToolEditor/Source/SubmitToolEditor/Private/SubmitToolEditorSettings.h b/Engine/Plugins/Experimental/SubmitToolEditor/Source/SubmitToolEditor/Private/SubmitToolEditorSettings.h
--- a/Engine/Plugins/Experimental/SubmitToolEditor/Source/SubmitToolEditor/Private/SubmitToolEditorSettings.h
+++ b/Engine/Plugins/Experimental/SubmitToolEditor/Source/SubmitToolEditor/Private/SubmitToolEditorSettings.h
@@ -8,8 +8,33 @@
 
 struct FPropertyChangedEvent;
 
-/** Settings for the submit tool in the editor */
+/** Settings for adding SubmitTool changelist tags based on validation messages */
+USTRUCT()
+struct FSubmitToolChangelistTagTrigger
+{
+	GENERATED_USTRUCT_BODY()
+
+	UPROPERTY(EditAnywhere, Category = "Submit Tool")
+	FString RegExMessage;
+
+	UPROPERTY(EditAnywhere, Category = "Submit Tool")
+	FString SubmitToolTag;
+};
+
+/** Settings for adding SubmitTool arguments based on validation messages */
+USTRUCT()
+struct FSubmitToolArgumentTrigger
+{
+	GENERATED_USTRUCT_BODY()
+
+	UPROPERTY(EditAnywhere, Category = "Submit Tool")
+	FString RegExMessage;
 
+	UPROPERTY(EditAnywhere, Category = "Submit Tool")
+	FString SubmitToolArgument;
+};
+
+/** Settings for the submit tool in the editor */
 UCLASS(config=Editor, meta = (DisplayName = "Submit Tool Settings"))
 class USubmitToolEditorSettings : public UDeveloperSettings
 {
@@ -34,5 +59,10 @@ public:
 	
 	UPROPERTY(config, Category = "Submit Tool", EditDefaultsOnly)
 	bool bEnforceDataValidation;
-};
 
+	UPROPERTY(config, Category = "Submit Tool", EditDefaultsOnly)
+	TArray<FSubmitToolChangelistTagTrigger> ChangelistTagTriggers;
+
+	UPROPERTY(config, Category = "Submit Tool", EditDefaultsOnly)
+	TArray<FSubmitToolArgumentTrigger> OptionalArgumentTriggers;
+};
