Fix new changelist created in editor not running data validation before invoking submit tool

#Synced-CL 48450364

diff --git a/Engine/Plugins/Experimental/SubmitToolEditor/Source/SubmitToolEditor/Private/SubmitToolEditor.cpp b/Engine/Plugins/Experimental/SubmitToolEditor/Source/SubmitToolEditor/Private/SubmitToolEditor.cpp
--- a/Engine/Plugins/Experimental/SubmitToolEditor/Source/SubmitToolEditor/Private/SubmitToolEditor.cpp
+++ b/Engine/Plugins/Experimental/SubmitToolEditor/Source/SubmitToolEditor/Private/SubmitToolEditor.cpp
@@ -233,7 +233,7 @@ FSubmitOverrideReply FSubmitToolEditorModule::InvokeSubmitTool(ISourceControlPro
 	FString Identifier;
 	if (CreateChangelist(InProvider, InDescription, InFiles, Identifier))
 	{
-		return InvokeSubmitTool(InProvider, InPath, InArgs, Identifier);
+		return InvokeSubmitTool(InProvider, InPath, InArgs, InDescription, Identifier);
 	}
 
 	return FSubmitOverrideReply::Error;
